- Adding sync verification in rivermax output sync policy to detect sync loss

diff --git a/Engine/Plugins/Runtime/nDisplay/Source/DisplayClusterMedia/Private/Synchronization/SyncPolicies/DisplayClusterMediaOutputSynchronizationPolicyEthernetBarrierBase.cpp b/Engine/Plugins/Runtime/nDisplay/Source/DisplayClusterMedia/Private/Synchronization/SyncPolicies/DisplayClusterMediaOutputSynchronizationPolicyEthernetBarrierBase.cpp
--- a/Engine/Plugins/Runtime/nDisplay/Source/DisplayClusterMedia/Private/Synchronization/SyncPolicies/DisplayClusterMediaOutputSynchronizationPolicyEthernetBarrierBase.cpp
+++ b/Engine/Plugins/Runtime/nDisplay/Source/DisplayClusterMedia/Private/Synchronization/SyncPolicies/DisplayClusterMediaOutputSynchronizationPolicyEthernetBarrierBase.cpp
@@ -107,6 +107,21 @@ FString FDisplayClusterMediaOutputSynchronizationPolicyEthernetBarrierBaseHandle
 	return MediaDeviceId;
 }
 
+IDisplayClusterGenericBarriersClient* const FDisplayClusterMediaOutputSynchronizationPolicyEthernetBarrierBaseHandler::GetBarrierClient() const
+{
+	return EthernetBarrierClient.Get();
+}
+
+const FString& FDisplayClusterMediaOutputSynchronizationPolicyEthernetBarrierBaseHandler::GetBarrierId() const
+{
+	return BarrierId;
+}
+
+const FString& FDisplayClusterMediaOutputSynchronizationPolicyEthernetBarrierBaseHandler::GetThreadMarker() const
+{
+	return ThreadMarker;
+}
+
 void FDisplayClusterMediaOutputSynchronizationPolicyEthernetBarrierBaseHandler::SyncThreadOnBarrier()
 {
 	// Sync on the barrier if everything is good
diff --git a/Engine/Plugins/Runtime/nDisplay/Source/DisplayClusterMedia/Public/Synchronization/DisplayClusterMediaOutputSynchronizationPolicyEthernetBarrierBase.h b/Engine/Plugins/Runtime/nDisplay/Source/DisplayClusterMedia/Public/Synchronization/DisplayClusterMediaOutputSynchronizationPolicyEthernetBarrierBase.h
--- a/Engine/Plugins/Runtime/nDisplay/Source/DisplayClusterMedia/Public/Synchronization/DisplayClusterMediaOutputSynchronizationPolicyEthernetBarrierBase.h
+++ b/Engine/Plugins/Runtime/nDisplay/Source/DisplayClusterMedia/Public/Synchronization/DisplayClusterMediaOutputSynchronizationPolicyEthernetBarrierBase.h
@@ -29,15 +29,26 @@ public:
 	virtual void Synchronize() = 0;
 
 protected:
+	
+	/** Initializes dynamic barrier on the primary node. */
+	virtual bool InitializeBarrier(const FString& SyncInstanceId);
+	
 	/** Synchronizes calling thread at the barrier. */
 	void SyncThreadOnBarrier();
 
 	/** Returns media device ID being synchronized */
 	FString GetMediaDeviceId() const;
 
+	/** Returns barrier client created for this sync policy */
+	IDisplayClusterGenericBarriersClient* const GetBarrierClient() const;
+
+	/** Get Barrier ID for this sync policy */
+	const FString& GetBarrierId() const;
+
+	/** Get thread marker for this sync policy */
+	const FString& GetThreadMarker() const;
+
 private:
-	/** Initializes dynamic barrier on the primary node. */
-	bool InitializeBarrier(const FString& SyncInstanceId);
 
 	/** Releases dynamic barrier on the primary node. */
 	void ReleaseBarrier();
diff --git a/Engine/Plugins/VirtualProduction/Rivermax/RivermaxCore/Source/RivermaxCore/Private/RivermaxOutputStream.cpp b/Engine/Plugins/VirtualProduction/Rivermax/RivermaxCore/Source/RivermaxCore/Private/RivermaxOutputStream.cpp
--- a/Engine/Plugins/VirtualProduction/Rivermax/RivermaxCore/Source/RivermaxCore/Private/RivermaxOutputStream.cpp
+++ b/Engine/Plugins/VirtualProduction/Rivermax/RivermaxCore/Source/RivermaxCore/Private/RivermaxOutputStream.cpp
@@ -124,6 +124,13 @@ namespace UE::RivermaxCore::Private
 		TEXT("Multiplier applied to desired output frame rate in order to reduce time it takes to send out a frame and slowly correct misalignment that could happen."),
 		ECVF_Default);
 
+	static bool GbTriggerRandomTimingIssue = false;
+	FAutoConsoleVariableRef CVarTriggerRandomTimingIssue(
+		TEXT("Rivermax.Sync.TriggerRandomTimingIssue")
+		, UE::RivermaxCore::Private::GbTriggerRandomTimingIssue
+		, TEXT("Randomly triggers a timing issue to test self repair."), ECVF_Cheat);
+
+
 	bool FindPayloadSize(const FRivermaxOutputStreamOptions& InOptions, uint32 InBytesPerLine, const FVideoFormatInfo& FormatInfo, uint16& OutPayloadSize)
 	{
 		using namespace UE::RivermaxCore::Private::Utils;
@@ -780,9 +787,9 @@ namespace UE::RivermaxCore::Private
 			{
 				TRACE_CPUPROFILER_EVENT_SCOPE(RmaxOut::WrappingUp);
 				if (CurrentFrame)
-				{
-					FrameManager->MarkAsSent(CurrentFrame);
-					CurrentFrame.Reset();
+				{	
+					constexpr bool bReleaseFrame = true;
+					CompleteCurrentFrame(bReleaseFrame);
 				}
 
 				// Make the next frame to send the current one and update its state
@@ -1482,7 +1489,8 @@ namespace UE::RivermaxCore::Private
 		// In frame creation alignment, we always release the last frame sent
 		if (CurrentFrame.IsValid())
 		{
-			FrameManager->MarkAsSent(CurrentFrame);
+			constexpr bool bReleaseFrame = true;
+			CompleteCurrentFrame(bReleaseFrame);
 		}
 
 		// Make the next frame to send the current one and update its state
@@ -1515,8 +1523,8 @@ namespace UE::RivermaxCore::Private
 		{
 			if (CurrentFrame)
 			{
-				FrameManager->MarkAsSent(CurrentFrame);
-				CurrentFrame.Reset();
+				constexpr bool bReleaseFrame = true;
+				CompleteCurrentFrame(bReleaseFrame);
 			}
 
 			// Make the next frame to send the current one and update its state
@@ -1529,6 +1537,10 @@ namespace UE::RivermaxCore::Private
 		}
 		else
 		{
+			// We finished sending a frame so complete it but don't release it as we will repeat it
+			constexpr bool bReleaseFrame = false;
+			CompleteCurrentFrame(bReleaseFrame);
+
 			// We will resend the last one so just reinitialize it to resend
 			InitializeNextFrame(CurrentFrame);
 			
@@ -1827,6 +1839,7 @@ namespace UE::RivermaxCore::Private
 		// Add Tro offset to next alignment point and configurable offset
 		StreamData.NextAlignmentPointNanosec = NextAlignmentNano;
 		StreamData.NextScheduleTimeNanosec = NextAlignmentNano + TransmitOffsetNanosec + CVarRivermaxScheduleOffset.GetValueOnAnyThread();
+		StreamData.LastAlignmentPointFrameNumber = StreamData.NextAlignmentPointFrameNumber;
 		StreamData.NextAlignmentPointFrameNumber = NextFrameNumber;
 
 		StreamData.bHasValidNextFrameNumber = bFoundValidTimings;
@@ -2057,6 +2070,19 @@ namespace UE::RivermaxCore::Private
 					// We skip an interval instead but it is quite drastic.
 					const bool bIsChunkOnTime = IsChunkOnTime();
 					CurrentFrame->bCaughtTimingIssue = !bIsChunkOnTime;
+
+					if (UE::RivermaxCore::Private::GbTriggerRandomTimingIssue)
+					{
+						FRandomStream RandomStream(FPlatformTime::Cycles64());
+						const bool bTriggerDesync = (RandomStream.RandRange(0.0, 1.0) > 0.7) ? true : false;
+						if (bTriggerDesync)
+						{
+							TRACE_CPUPROFILER_EVENT_SCOPE(RmaxOut::ForceTimingIssue);
+							CurrentFrame->bCaughtTimingIssue = true;
+						}
+
+						UE::RivermaxCore::Private::GbTriggerRandomTimingIssue = false;
+					}
 				}
 
 				CommitNextChunks();
@@ -2097,5 +2123,29 @@ namespace UE::RivermaxCore::Private
 		}
 	}
 
+	void FRivermaxOutputStream::GetLastPresentedFrame(FPresentedFrameInfo& OutFrameInfo) const
+	{
+		FScopeLock Lock(&PresentedFrameCS);
+		OutFrameInfo = LastPresentedFrame;
+	}
+
+	void FRivermaxOutputStream::CompleteCurrentFrame(bool bReleaseFrame)
+	{
+		if (ensure(CurrentFrame))
+		{
+			{
+				FScopeLock Lock(&PresentedFrameCS);
+				LastPresentedFrame.RenderedFrameNumber = CurrentFrame->FrameIdentifier;
+				LastPresentedFrame.PresentedFrameBoundaryNumber = StreamData.LastAlignmentPointFrameNumber;
+			}
+			
+			// We don't release when there is no new frame, so we keep a hold on it to repeat it.
+			if (bReleaseFrame)
+			{
+				FrameManager->MarkAsSent(CurrentFrame);
+				CurrentFrame.Reset();
+			}
+		}
+	}
 }
 
diff --git a/Engine/Plugins/VirtualProduction/Rivermax/RivermaxCore/Source/RivermaxCore/Private/RivermaxOutputStream.h b/Engine/Plugins/VirtualProduction/Rivermax/RivermaxCore/Source/RivermaxCore/Private/RivermaxOutputStream.h
--- a/Engine/Plugins/VirtualProduction/Rivermax/RivermaxCore/Source/RivermaxCore/Private/RivermaxOutputStream.h
+++ b/Engine/Plugins/VirtualProduction/Rivermax/RivermaxCore/Source/RivermaxCore/Private/RivermaxOutputStream.h
@@ -169,6 +169,9 @@ namespace UE::RivermaxCore::Private
 		/** Next alignment point frame number treated to detect missed frames */
 		uint64 NextAlignmentPointFrameNumber = 0;
 
+		/** Last alignment point frame number we have processed*/
+		uint64 LastAlignmentPointFrameNumber = 0;
+
 		/** Timestamp at which we started commiting a frame */
 		uint64 LastSendStartTimeNanoSec = 0;
 		
@@ -222,6 +225,7 @@ namespace UE::RivermaxCore::Private
 		virtual bool PushVideoFrame(const FRivermaxOutputVideoFrameInfo& NewFrame) override;
 		virtual bool IsGPUDirectSupported() const override;
 		virtual bool ReserveFrame(uint32 FrameIdentifier) const override;
+		virtual void GetLastPresentedFrame(FPresentedFrameInfo& OutFrameInfo) const override;
 		//~ End IRivermaxOutputStream interface
 
 		void Process_AnyThread();
@@ -342,6 +346,9 @@ namespace UE::RivermaxCore::Private
 		/** Go through all chunks of current frame and commit them to Rivermax to send them at the next desired time */
 		void SendFrame();
 
+		/** When a frame has been sent (after frame interval), we update last presented frame tracking and optionally release it in the presentation queue */
+		void CompleteCurrentFrame(bool bReleaseFrame);
+
 	private:
 
 		/** Options related to this stream. i.e resolution, frame rate, etc... */
@@ -359,9 +366,6 @@ namespace UE::RivermaxCore::Private
 		/** Stream id returned by rmax library */
 		rmx_stream_id StreamId = 0;
 
-		/** Critical section to protect frames access */
-		mutable FCriticalSection FrameCriticalSection;
-
 		/** Current frame being sent */
 		TSharedPtr<FRivermaxOutputFrame> CurrentFrame;
 
@@ -418,6 +422,12 @@ namespace UE::RivermaxCore::Private
 		/** Whether to trigger a delay in the output thread loop next time it ticks */
 		bool bTriggerRandomDelay = false;
 
+		/** Critical section to access data of last presented frame */
+		mutable FCriticalSection PresentedFrameCS;
+
+		/** Info of last presented frame */
+		FPresentedFrameInfo LastPresentedFrame;
+
 		friend struct FRTPHeaderPrefiller;
 	};
 }
diff --git a/Engine/Plugins/VirtualProduction/Rivermax/RivermaxCore/Source/RivermaxCore/Public/IRivermaxOutputStream.h b/Engine/Plugins/VirtualProduction/Rivermax/RivermaxCore/Source/RivermaxCore/Public/IRivermaxOutputStream.h
--- a/Engine/Plugins/VirtualProduction/Rivermax/RivermaxCore/Source/RivermaxCore/Public/IRivermaxOutputStream.h
+++ b/Engine/Plugins/VirtualProduction/Rivermax/RivermaxCore/Source/RivermaxCore/Public/IRivermaxOutputStream.h
@@ -11,16 +11,38 @@ namespace UE::RivermaxCore
 {
 	struct FRivermaxOutputStreamOptions;
 
+	/** Description of video frame to be captured. */
 	struct RIVERMAXCORE_API FRivermaxOutputVideoFrameInfo
 	{
+		/** Incremental number identifying frame. Usually GFrameCounter */
 		uint32 FrameIdentifier = 0;
+
+		/** Height of the frame */
 		uint32 Height = 0;
+		
+		/** Width of the frame */
 		uint32 Width = 0;
+
+		/** Stride of a line */
 		uint32 Stride = 0;
+
+		/** Video frame data location in system memory */
 		void* VideoBuffer = nullptr;
+
+		/** Video frame data location when using GPUDirect */
 		FBufferRHIRef GPUBuffer;
 	};
 
+	/** Information about the last frame that was presented by the stream */
+	struct RIVERMAXCORE_API FPresentedFrameInfo
+	{
+		/** Frame boundary at which the RenderedFrameNumber has been presented*/
+		uint64 PresentedFrameBoundaryNumber = 0;
+
+		/** Last engine's FrameNumber that was presented. */
+		uint32 RenderedFrameNumber = 0;
+	};
+
 	class RIVERMAXCORE_API IRivermaxOutputStreamListener
 	{
 	public:
@@ -70,6 +92,9 @@ namespace UE::RivermaxCore
 		 * this method will block until a free frame is found. 
 		 */
 		virtual bool ReserveFrame(uint32 FrameIdentifier) const = 0;
+
+		/** Returns information about the last frame that was presented on the wire */
+		virtual void GetLastPresentedFrame(FPresentedFrameInfo& OutFrameInfo) const = 0;
 	};
 }
 
diff --git a/Engine/Plugins/VirtualProduction/Rivermax/RivermaxMedia/Source/RivermaxMedia/Private/RivermaxMediaCapture.cpp b/Engine/Plugins/VirtualProduction/Rivermax/RivermaxMedia/Source/RivermaxMedia/Private/RivermaxMediaCapture.cpp
--- a/Engine/Plugins/VirtualProduction/Rivermax/RivermaxMedia/Source/RivermaxMedia/Private/RivermaxMediaCapture.cpp
+++ b/Engine/Plugins/VirtualProduction/Rivermax/RivermaxMedia/Source/RivermaxMedia/Private/RivermaxMediaCapture.cpp
@@ -126,6 +126,14 @@ UE::RivermaxCore::FRivermaxOutputStreamOptions URivermaxMediaCapture::GetOutputS
 	return Options;
 }
 
+void URivermaxMediaCapture::GetLastPresentedFrameInformation(UE::RivermaxCore::FPresentedFrameInfo& OutFrameInfo) const
+{
+	if (RivermaxStream)
+	{
+		RivermaxStream->GetLastPresentedFrame(OutFrameInfo);
+	}
+}
+
 bool URivermaxMediaCapture::ValidateMediaOutput() const
 {
 	URivermaxMediaOutput* RivermaxMediaOutput = Cast<URivermaxMediaOutput>(MediaOutput);
diff --git a/Engine/Plugins/VirtualProduction/Rivermax/RivermaxMedia/Source/RivermaxMedia/Public/RivermaxMediaCapture.h b/Engine/Plugins/VirtualProduction/Rivermax/RivermaxMedia/Source/RivermaxMedia/Public/RivermaxMediaCapture.h
--- a/Engine/Plugins/VirtualProduction/Rivermax/RivermaxMedia/Source/RivermaxMedia/Public/RivermaxMediaCapture.h
+++ b/Engine/Plugins/VirtualProduction/Rivermax/RivermaxMedia/Source/RivermaxMedia/Public/RivermaxMediaCapture.h
@@ -26,6 +26,9 @@ public:
 	/** Rivermax capture specific API to provide stream options access */
 	UE::RivermaxCore::FRivermaxOutputStreamOptions GetOutputStreamOptions() const;
 
+	/** Returns information about last presented frame on the output stream */
+	void GetLastPresentedFrameInformation(UE::RivermaxCore::FPresentedFrameInfo& OutFrameInfo) const;
+
 public:
 
 	//~ Begin UMediaCapture interface
diff --git a/Engine/Plugins/VirtualProduction/Rivermax/RivermaxSync/Source/RivermaxSync/Private/MediaOutputSynchronizationPolicyRivermax.cpp b/Engine/Plugins/VirtualProduction/Rivermax/RivermaxSync/Source/RivermaxSync/Private/MediaOutputSynchronizationPolicyRivermax.cpp
--- a/Engine/Plugins/VirtualProduction/Rivermax/RivermaxSync/Source/RivermaxSync/Private/MediaOutputSynchronizationPolicyRivermax.cpp
+++ b/Engine/Plugins/VirtualProduction/Rivermax/RivermaxSync/Source/RivermaxSync/Private/MediaOutputSynchronizationPolicyRivermax.cpp
@@ -2,16 +2,47 @@
 
 #include "MediaOutputSynchronizationPolicyRivermax.h"
 
+#include "Cluster/IDisplayClusterClusterManager.h"
+#include "HAL/IConsoleManager.h"
+#include "IDisplayCluster.h"
 #include "IRivermaxCoreModule.h"
 #include "IRivermaxManager.h"
 #include "RivermaxMediaCapture.h"
 #include "RivermaxPTPUtils.h"
+#include "RivermaxSyncLog.h"
 
 
+
+namespace UE::RivermaxSync
+{
+	static TAutoConsoleVariable<float> CVarRivermaxSyncWakeupOffset(
+		TEXT("Rivermax.Sync.WakeUpOffset"), 0.5f,
+		TEXT("Offset from alignment point to wake up at when barrier stalls the cluster. Units: milliseconds"),
+		ECVF_Default);
+
+	static TAutoConsoleVariable<bool> CVarRivermaxSyncEnableSelfRepair(
+		TEXT("Rivermax.Sync.EnableSelfRepair"), true,
+		TEXT("Whether to use exchanged data in the synchronization barrier to detect desynchronized state and act on it to self repair"),
+		ECVF_Default);
+
+	static TAutoConsoleVariable<int32> CVarRivermaxSyncMaxFrameTimeRange(
+		TEXT("Rivermax.Sync.MaxFrameTimeRange"), 8,
+		TEXT("Maximum number of frame last presented frame from nodes can be from each other when they join the barrier to enable self repair"),
+		ECVF_Default);
+
+	static bool GbTriggerRandomDesync = false;
+	FAutoConsoleVariableRef CVarTriggerRandomDesync(
+		TEXT("Rivermax.Sync.ForceDesync")
+		, UE::RivermaxSync::GbTriggerRandomDesync
+		, TEXT("After barrier synchronization, trigger random stall."), ECVF_Cheat);
+}
+
 FMediaOutputSynchronizationPolicyRivermaxHandler::FMediaOutputSynchronizationPolicyRivermaxHandler(UMediaOutputSynchronizationPolicyRivermax* InPolicyObject)
 	: Super(InPolicyObject)
+	, MarginMs(InPolicyObject->MarginMs)
 {
-
+	// Allocate memory to store the data exchanged in the barrier
+	BarrierData.SetNumZeroed(sizeof(FMediaSyncBarrierData));
 }
 
 TSubclassOf<UDisplayClusterMediaOutputSynchronizationPolicy> FMediaOutputSynchronizationPolicyRivermaxHandler::GetPolicyClass() const
@@ -48,7 +79,95 @@ bool FMediaOutputSynchronizationPolicyRivermaxHandler::IsCaptureTypeSupported(UM
 	return false;
 }
 
-double FMediaOutputSynchronizationPolicyRivermaxHandler::GetTimeBeforeNextSyncPoint()
+void FMediaOutputSynchronizationPolicyRivermaxHandler::Synchronize()
+{
+	{
+		TRACE_CPUPROFILER_EVENT_SCOPE(RmaxSync::Synchronize);
+
+		// Sync on the barrier if everything is good
+		if (!IsRunning())
+		{
+			UE_LOG(LogRivermaxSync, Warning, TEXT("'%s': Synchronization is off"), *GetMediaDeviceId());
+			return;
+		}
+
+		IDisplayClusterGenericBarriersClient* const BarrierClient = GetBarrierClient();
+		if (!BarrierClient || !BarrierClient->IsConnected())
+		{
+			UE_LOG(LogRivermaxSync, Warning, TEXT("'%s': Barrier client is not connected or nullptr"), *GetMediaDeviceId());
+			return;
+		}
+
+		UE_LOG(LogRivermaxSync, Verbose, TEXT("'%s': Synchronizing caller '%s' at the barrier '%s'"), *GetMediaDeviceId(), *GetThreadMarker(), *GetBarrierId());
+
+		URivermaxMediaCapture* RmaxCapture = Cast<URivermaxMediaCapture>(CapturingDevice);
+		if (RmaxCapture == nullptr || RmaxCapture->GetState() != EMediaCaptureState::Capturing)
+		{
+			UE_LOG(LogRivermaxSync, Warning, TEXT("'%s': Rivermax Capture isn't valid or not capturing"), *GetMediaDeviceId());
+			return;
+		}
+
+		// Verify if we are safe to go inside the barrier.
+		{
+			// Ask the sync implementation about how much time we have before next synchronization timepoint
+			const double TimeLeftSeconds = GetTimeBeforeNextSyncPoint();
+			// Convert to seconds
+			const double MarginSeconds = double(MarginMs) / 1000;
+
+			// In case we're unsafe, skip the upcoming sync timepoint
+			if (TimeLeftSeconds < MarginSeconds)
+			{
+				TRACE_CPUPROFILER_EVENT_SCOPE(RmaxSync::MarginProtection);
+
+				const float OffsetTimeSeconds = UE::RivermaxSync::CVarRivermaxSyncWakeupOffset.GetValueOnAnyThread() * 1E-3;
+				// Sleep for a bit longer to skip the alignment timepoint
+				const float SleepTime = TimeLeftSeconds + OffsetTimeSeconds;
+
+				UE_LOG(LogRivermaxSync, VeryVerbose, TEXT("'%s': TimeLeft(%lf) < Margin(%lf) --> Sleeping for %lf..."),
+					*GetMediaDeviceId(), TimeLeftSeconds, MarginSeconds, SleepTime);
+
+				FPlatformProcess::SleepNoStats(SleepTime);
+			}
+		}
+
+
+		// We are good to go in the barrier, prepare payload about presented frame
+		UE::RivermaxCore::FPresentedFrameInfo FrameInfo;
+		RmaxCapture->GetLastPresentedFrameInformation(FrameInfo);
+
+		// Fill the memory to be exchanged by nodes in the barrier.
+		FMediaSyncBarrierData BarrierDataStruct(FrameInfo);
+		UE_LOG(LogRivermaxSync, VeryVerbose, TEXT("'%s' Entering with %u"), *GetMediaDeviceId(), BarrierDataStruct.LastRenderedFrameNumber);
+		FMemory::Memcpy(BarrierData.GetData(), &BarrierDataStruct, sizeof(FMediaSyncBarrierData));
+
+		// We don't use response data for now
+		TArray<uint8> ResponseData;
+
+		// Synchronize on a barrier
+		BarrierClient->Synchronize(GetBarrierId(), GetThreadMarker(), BarrierData, ResponseData);
+	}
+
+	// Debug cvar to potentially stall a node after exiting the barrier and missing alignment points
+	if (UE::RivermaxSync::GbTriggerRandomDesync)
+	{
+		FRandomStream RandomStream(FPlatformTime::Cycles64());
+		const bool bTriggerDesync = (RandomStream.RandRange(0.0, 1.0) > 0.7) ? true : false;
+		if (bTriggerDesync)
+		{
+			TRACE_CPUPROFILER_EVENT_SCOPE(RmaxSync::ForceBadSync);
+			URivermaxMediaCapture* RmaxCapture = Cast<URivermaxMediaCapture>(CapturingDevice);
+			double TimeLeftSeconds = GetTimeBeforeNextSyncPoint();
+			TimeLeftSeconds += RmaxCapture->GetOutputStreamOptions().FrameRate.AsInterval();
+			const float OffsetTimeSeconds = UE::RivermaxSync::CVarRivermaxSyncWakeupOffset.GetValueOnAnyThread() * 1E-3;
+			const float SleepTime = TimeLeftSeconds + OffsetTimeSeconds;
+			FPlatformProcess::SleepNoStats(SleepTime);
+		}
+
+		UE::RivermaxSync::GbTriggerRandomDesync = false;
+	}
+}
+
+double FMediaOutputSynchronizationPolicyRivermaxHandler::GetTimeBeforeNextSyncPoint() const
 {
 	if (URivermaxMediaCapture* RmaxCapture = Cast<URivermaxMediaCapture>(CapturingDevice))
 	{
@@ -80,6 +199,143 @@ double FMediaOutputSynchronizationPolicyRivermaxHandler::GetTimeBeforeNextSyncPo
 	return 1.f;
 }
 
+bool FMediaOutputSynchronizationPolicyRivermaxHandler::InitializeBarrier(const FString& SyncInstanceId)
+{
+	// Base initialization first
+	if (!Super::InitializeBarrier(SyncInstanceId))
+	{
+		UE_LOG(LogRivermaxSync, Warning, TEXT("Couldn't initialize barrier for '%s'"), *GetMediaDeviceId());
+		return false;
+	}
+
+	// Only setup barrier delegate for primary node to handle barrier data verification
+	if (!IDisplayCluster::Get().GetClusterMgr()->IsPrimary())
+	{
+		return true;
+	}
+
+	// Get barrier client
+	IDisplayClusterGenericBarriersClient* const BarrierClient = GetBarrierClient();
+	if (!BarrierClient || !BarrierClient->IsConnected())
+	{
+		UE_LOG(LogRivermaxSync, Warning, TEXT("Couldn't access a barrier client for '%s'"), *GetMediaDeviceId());
+		return false;
+	}
+
+	// Get delegate bound to the specific barrier
+	IDisplayClusterGenericBarriersClient::FOnGenericBarrierSynchronizationDelegate* Delegate = BarrierClient->GetBarrierSyncDelegate(GetBarrierId());
+	if (!Delegate)
+	{
+		UE_LOG(LogRivermaxSync, Warning, TEXT("'%s': Couldn't access a barrier delegate for barrier '%s'"), *GetMediaDeviceId(), *GetBarrierId());
+		return false;
+	}
+
+	// Setup synchronization delegate that will be called on the p-node
+	Delegate->BindRaw(this, &FMediaOutputSynchronizationPolicyRivermaxHandler::HandleBarrierSync);
+
+	return true;
+}
+
+void FMediaOutputSynchronizationPolicyRivermaxHandler::HandleBarrierSync(FGenericBarrierSynchronizationDelegateData& BarrierSyncData)
+{
+	TRACE_CPUPROFILER_EVENT_SCOPE(RmaxSync::BarrierSync);
+
+	if (bHasVerifiedClocks == false)
+	{
+		bHasVerifiedClocks = true;
+		bCanUseSelfRepair = ValidateNodesFrameTime(BarrierSyncData.RequestData);
+	}
+
+	if (!bCanUseSelfRepair || UE::RivermaxSync::CVarRivermaxSyncEnableSelfRepair.GetValueOnAnyThread() == false)
+	{
+		return;
+	}
+
+	if (BarrierSyncData.RequestData.Num() <= 0)
+	{
+		UE_LOG(LogRivermaxSync, Verbose, TEXT("'%s': No data was provided by nodes for sync barrier."), *GetMediaDeviceId())
+		return;
+	}
+
+	FString FirstFrameNode;
+	FMediaSyncBarrierData FirstNodeData;
+	bool bFirstNode = true;
+	bool bSelfRepairRequired = false;
+
+	for (const TPair<FString, TArray<uint8>>& NodePresentedFrame : BarrierSyncData.RequestData)
+	{
+		check(NodePresentedFrame.Value.Num() == sizeof(FMediaSyncBarrierData));
+
+		const FMediaSyncBarrierData* const NodeData = reinterpret_cast<const FMediaSyncBarrierData* const>(NodePresentedFrame.Value.GetData());
+		
+		// We expect all nodes to enter the barrier after presenting the SAME frame at the SAME frame number.
+		// If a node enters the barrier a frame late on the other, self repair will be triggered.
+		// Samething goes if all nodes didn't present the same frame
+		if (bFirstNode)
+		{
+			bFirstNode = false;
+			FirstFrameNode = NodePresentedFrame.Key;
+			FirstNodeData = *NodeData;
+		}
+		else if((FirstNodeData.LastRenderedFrameNumber != NodeData->LastRenderedFrameNumber)
+			 || (FirstNodeData.PresentedFrameBoundaryNumber != NodeData->PresentedFrameBoundaryNumber))
+		{
+			UE_LOG(LogRivermaxSync, Warning, TEXT("Desync detected: Node '%s' presented frame '%u' at frame boundary '%llu', but node '%s' presented frame '%u' at frame boundary '%llu'")
+			, *FirstFrameNode
+			, FirstNodeData.LastRenderedFrameNumber
+			, FirstNodeData.PresentedFrameBoundaryNumber
+			, *NodePresentedFrame.Key
+			, NodeData->LastRenderedFrameNumber
+			, NodeData->PresentedFrameBoundaryNumber);
+
+			bSelfRepairRequired = true;
+			break;
+		}
+	}
+
+	// If repair is required, we stall until we are past the next alignment point to have all scheduler present something and get closer to a synchronized state.
+	if (bSelfRepairRequired)
+	{
+		TRACE_CPUPROFILER_EVENT_SCOPE(RmaxSync::SelfRepair);
+
+		const double TimeLeftSeconds = GetTimeBeforeNextSyncPoint();
+		const float OffsetTimeSeconds = UE::RivermaxSync::CVarRivermaxSyncWakeupOffset.GetValueOnAnyThread() * 1E-3;
+		const float SleepTime = TimeLeftSeconds + OffsetTimeSeconds;
+		FPlatformProcess::SleepNoStats(SleepTime);
+	}
+	else
+	{
+		UE_LOG(LogRivermaxSync, VeryVerbose, TEXT("'%s': Cluster synchronized. Nodes presented frame %u at frame boundary %llu"), *GetMediaDeviceId(), FirstNodeData.LastRenderedFrameNumber, FirstNodeData.PresentedFrameBoundaryNumber);
+	}
+}
+
+bool FMediaOutputSynchronizationPolicyRivermaxHandler::ValidateNodesFrameTime(const TMap<FString, TArray<uint8>>& NodeRequestData) const
+{
+	uint64 MinFrameTime = TNumericLimits<uint64>::Max();
+	uint64 MaxFrameTime = TNumericLimits<uint64>::Min();
+	for (const TPair<FString, TArray<uint8>>& NodePresentedFrame : NodeRequestData)
+	{
+		check(NodePresentedFrame.Value.Num() == sizeof(FMediaSyncBarrierData));
+
+		const FMediaSyncBarrierData* const NodeData = reinterpret_cast<const FMediaSyncBarrierData* const>(NodePresentedFrame.Value.GetData());
+		MinFrameTime = FMath::Min(MinFrameTime, NodeData->PresentedFrameBoundaryNumber);
+		MaxFrameTime = FMath::Max(MaxFrameTime, NodeData->PresentedFrameBoundaryNumber);
+	}
+
+	const uint64 MaxRange = UE::RivermaxSync::CVarRivermaxSyncMaxFrameTimeRange.GetValueOnAnyThread();
+	const uint64 DetectedRange = MaxFrameTime - MinFrameTime;
+	if (DetectedRange > MaxRange)
+	{	
+		UE_LOG(LogRivermaxSync, Warning, TEXT("Self repair can't be enabled. Frame time range (%llu) across cluster too large. Verify PTP clocks to be identical on each node."), DetectedRange);
+		return false;
+	}
+	else
+	{
+		UE_LOG(LogRivermaxSync, Log, TEXT("Self repair can be used. Frame time range (%llu) across cluster points to a common time reference."), DetectedRange);
+		return true;
+	}
+}
+
 TSharedPtr<IDisplayClusterMediaOutputSynchronizationPolicyHandler> UMediaOutputSynchronizationPolicyRivermax::GetHandler()
 {
 	if (!Handler)
@@ -89,3 +345,4 @@ TSharedPtr<IDisplayClusterMediaOutputSynchronizationPolicyHandler> UMediaOutputS
 
 	return Handler;
 }
+
diff --git a/Engine/Plugins/VirtualProduction/Rivermax/RivermaxSync/Source/RivermaxSync/Private/RivermaxSyncLog.cpp b/Engine/Plugins/VirtualProduction/Rivermax/RivermaxSync/Source/RivermaxSync/Private/RivermaxSyncLog.cpp
new file mode 100644
--- /dev/null
+++ b/Engine/Plugins/VirtualProduction/Rivermax/RivermaxSync/Source/RivermaxSync/Private/RivermaxSyncLog.cpp
@@ -0,0 +1,5 @@
+// Copyright Epic Games, Inc. All Rights Reserved.
+
+#include "RivermaxSyncLog.h"
+
+DEFINE_LOG_CATEGORY(LogRivermaxSync);
diff --git a/Engine/Plugins/VirtualProduction/Rivermax/RivermaxSync/Source/RivermaxSync/Private/RivermaxSyncLog.h b/Engine/Plugins/VirtualProduction/Rivermax/RivermaxSync/Source/RivermaxSync/Private/RivermaxSyncLog.h
new file mode 100644
--- /dev/null
+++ b/Engine/Plugins/VirtualProduction/Rivermax/RivermaxSync/Source/RivermaxSync/Private/RivermaxSyncLog.h
@@ -0,0 +1,7 @@
+// Copyright Epic Games, Inc. All Rights Reserved.
+
+#pragma once
+
+#include "Logging/LogMacros.h"
+
+DECLARE_LOG_CATEGORY_EXTERN(LogRivermaxSync, Log, All);
diff --git a/Engine/Plugins/VirtualProduction/Rivermax/RivermaxSync/Source/RivermaxSync/Public/MediaOutputSynchronizationPolicyRivermax.h b/Engine/Plugins/VirtualProduction/Rivermax/RivermaxSync/Source/RivermaxSync/Public/MediaOutputSynchronizationPolicyRivermax.h
--- a/Engine/Plugins/VirtualProduction/Rivermax/RivermaxSync/Source/RivermaxSync/Public/MediaOutputSynchronizationPolicyRivermax.h
+++ b/Engine/Plugins/VirtualProduction/Rivermax/RivermaxSync/Source/RivermaxSync/Public/MediaOutputSynchronizationPolicyRivermax.h
@@ -2,17 +2,20 @@
 
 #pragma once
 
-#include "Synchronization/DisplayClusterMediaOutputSynchronizationPolicyThresholdBase.h"
+#include "Synchronization/DisplayClusterMediaOutputSynchronizationPolicyEthernetBarrierBase.h"
+
+#include "IRivermaxOutputStream.h"
 
 #include "MediaOutputSynchronizationPolicyRivermax.generated.h"
 
+
 /*
  * Synchronization logic handler class for UMediaOutputSynchronizationPolicyRivermax.
  */
 class RIVERMAXSYNC_API FMediaOutputSynchronizationPolicyRivermaxHandler
-	: public FDisplayClusterMediaOutputSynchronizationPolicyThresholdBaseHandler
+	: public FDisplayClusterMediaOutputSynchronizationPolicyEthernetBarrierBaseHandler
 {
-	using Super = FDisplayClusterMediaOutputSynchronizationPolicyThresholdBaseHandler;
+	using Super = FDisplayClusterMediaOutputSynchronizationPolicyEthernetBarrierBaseHandler;
 public:
 
 	FMediaOutputSynchronizationPolicyRivermaxHandler(UMediaOutputSynchronizationPolicyRivermax* InPolicyObject);
@@ -22,12 +25,59 @@ public:
 	//~ End IDisplayClusterMediaOutputSynchronizationPolicyHandler interface
 
 public:
+	/** We do our own synchronization by looking at distance to alignment point. */
+	virtual void Synchronize() override;
+
 	/** Returns true if specified media capture type can be synchonized by the policy implementation */
 	virtual bool IsCaptureTypeSupported(UMediaCapture* MediaCapture) const override;
 
 protected:
+	/** Initializes dynamic barrier on the primary node. */
+	virtual bool InitializeBarrier(const FString& SyncInstanceId) override;
+
+	/** Barrier callback containing data from each node to detect if cluster is out of sync. */
+	void HandleBarrierSync(FGenericBarrierSynchronizationDelegateData& BarrierSyncData);
+
+	/** Returns true if frame time from each node are close to each other (limit controlled by cvar) */
+	bool ValidateNodesFrameTime(const TMap<FString, TArray<uint8>>& NodeRequestData) const;
+
 	/** Returns amount of time before next synchronization point. */
-	virtual double GetTimeBeforeNextSyncPoint() override;
+	double GetTimeBeforeNextSyncPoint() const;
+
+protected:
+
+	/** Holds data provided to server by each node when joining the barrier */
+	struct FMediaSyncBarrierData
+	{
+		FMediaSyncBarrierData()
+		{
+		}
+
+		/** For now, we build it directly based on presentation info from the stream */
+		FMediaSyncBarrierData(const UE::RivermaxCore::FPresentedFrameInfo& FrameInfo)
+			: PresentedFrameBoundaryNumber(FrameInfo.PresentedFrameBoundaryNumber)
+			, LastRenderedFrameNumber(FrameInfo.RenderedFrameNumber)
+		{
+		}
+
+		/** Frame boundary number at which the last frame was presented.  */
+		uint64 PresentedFrameBoundaryNumber = 0;
+		
+		/** Last engine frame number that was presented */
+		uint32 LastRenderedFrameNumber = 0;
+	};
+
+	/** Synchronization margin (ms) */
+	float MarginMs = 5.0f;
+
+	/** Memory buffer used to contain data exchanged in the barrier */
+	TArray<uint8> BarrierData;
+
+	/** On first barrier sync, we will verify frame time range to detect if clocks are too far appart. Even when desynchronized, frame time shouldn't be further than 1-2 frame away */
+	bool bCanUseSelfRepair = true;
+
+	/** Verify clock between nodes only once */
+	bool bHasVerifiedClocks = false;
 };
 
 
@@ -36,7 +86,7 @@ protected:
  */
 UCLASS(editinlinenew, Blueprintable, meta = (DisplayName = "Rivermax (PTP)"))
 class RIVERMAXSYNC_API UMediaOutputSynchronizationPolicyRivermax
-	: public UDisplayClusterMediaOutputSynchronizationPolicyThresholdBase
+	: public UDisplayClusterMediaOutputSynchronizationPolicyEthernetBarrierBase
 {
 	GENERATED_BODY()
 
@@ -46,4 +96,9 @@ public:
 
 protected:
 	TSharedPtr<IDisplayClusterMediaOutputSynchronizationPolicyHandler> Handler;
+
+public:
+	/** Synchronization margin (ms) */
+	UPROPERTY(EditAnywhere, BlueprintReadWrite, Category = "Settings", meta = (DisplayName = "Margin (ms)", ClampMin = "1", ClampMax = "20", UIMin = "1", UIMax = "20"))
+	float MarginMs = 5.0f;
 };
diff --git a/Engine/Plugins/VirtualProduction/Rivermax/RivermaxSync/Source/RivermaxSync/RivermaxSync.Build.cs b/Engine/Plugins/VirtualProduction/Rivermax/RivermaxSync/Source/RivermaxSync/RivermaxSync.Build.cs
--- a/Engine/Plugins/VirtualProduction/Rivermax/RivermaxSync/Source/RivermaxSync/RivermaxSync.Build.cs
+++ b/Engine/Plugins/VirtualProduction/Rivermax/RivermaxSync/Source/RivermaxSync/RivermaxSync.Build.cs
@@ -8,14 +8,14 @@ public class RivermaxSync : ModuleRules
 	{
 		PublicDependencyModuleNames.AddRange(
 			new string[] {
-				"DisplayClusterMedia"
+				"DisplayClusterMedia",
+				"RivermaxCore"
 			});
 
 		PrivateDependencyModuleNames.AddRange(
 			new string[] {
 				"Core",
 				"CoreUObject",
-				"RivermaxCore",
 				"RivermaxMedia"
 			});
 
