Rivermax: Fix double skipped frame due to false de-sync detection.

The false de-sync detection was happening when the barriers are entered on the nodes at different times with respect to ptp vsync, such that theit latest frame presentation does not correspond to the same vsync.

As a solution, the barrier now carries the last to frame presentation datas, and a desync is only confirmed when the same vsync presented different frames. Otherwise it conservativeliy does nothing.

diff --git a/Engine/Plugins/VirtualProduction/Rivermax/RivermaxSync/Source/RivermaxSync/Private/MediaOutputSynchronizationPolicyRivermax.cpp b/Engine/Plugins/VirtualProduction/Rivermax/RivermaxSync/Source/RivermaxSync/Private/MediaOutputSynchronizationPolicyRivermax.cpp
--- a/Engine/Plugins/VirtualProduction/Rivermax/RivermaxSync/Source/RivermaxSync/Private/MediaOutputSynchronizationPolicyRivermax.cpp
+++ b/Engine/Plugins/VirtualProduction/Rivermax/RivermaxSync/Source/RivermaxSync/Private/MediaOutputSynchronizationPolicyRivermax.cpp
@@ -136,7 +136,7 @@ void FMediaOutputSynchronizationPolicyRivermaxHandler::Synchronize()
 		RmaxCapture->GetLastPresentedFrameInformation(FrameInfo);
 
 		// Fill the memory to be exchanged by nodes in the barrier.
-		FMediaSyncBarrierData BarrierDataStruct(FrameInfo);
+		BarrierDataStruct.InsertFrameInfo(FrameInfo);
 		UE_LOG(LogRivermaxSync, VeryVerbose, TEXT("'%s' Entering with %u"), *GetMediaDeviceId(), BarrierDataStruct.LastRenderedFrameNumber);
 		FMemory::Memcpy(BarrierData.GetData(), &BarrierDataStruct, sizeof(FMediaSyncBarrierData));
 
@@ -236,6 +236,75 @@ bool FMediaOutputSynchronizationPolicyRivermaxHandler::InitializeBarrier(const F
 	return true;
 }
 
+
+bool FMediaOutputSynchronizationPolicyRivermaxHandler::FMediaSyncBarrierData::HasConfirmedDesync(const FMediaSyncBarrierData& OtherBarrierData) const
+{
+	for (int32 FirstNodeIdx = 0; FirstNodeIdx < FMediaSyncBarrierData::FRAMEHISTORYLEN; ++FirstNodeIdx)
+	{
+		const uint32 FirstNodeFrameNumber = LastRenderedFrameNumber[FirstNodeIdx];
+		const uint64 FirstNodeVsyncBoundary = PresentedFrameBoundaryNumber[FirstNodeIdx];
+
+		for (int32 OtherNodeIdx = 0; OtherNodeIdx < FMediaSyncBarrierData::FRAMEHISTORYLEN; ++OtherNodeIdx)
+		{
+			const uint32 OtherNodeFrameNumber = OtherBarrierData.LastRenderedFrameNumber[OtherNodeIdx];
+			const uint64 OtherNodeVsyncBoundary = OtherBarrierData.PresentedFrameBoundaryNumber[OtherNodeIdx];
+
+			const bool bSameFrame = (FirstNodeFrameNumber == OtherNodeFrameNumber);
+			const bool bSameVsync = (FirstNodeVsyncBoundary == OtherNodeVsyncBoundary);
+
+			// If they agree on a recent frame, consider them in sync
+			if (bSameFrame && bSameVsync)
+			{
+				return false;
+			}
+
+			// If they presented a different frame on the same vsync, they are out of sync
+			if (bSameVsync && !bSameFrame)
+			{
+				return true;
+			}
+		}
+	}
+
+	// If we could not confirm a desync, we conservatively not consider a desync to have been detected.
+	return false;
+}
+
+
+FString FMediaOutputSynchronizationPolicyRivermaxHandler::FMediaSyncBarrierData::LastRenderedFrameNumbersAsString() const
+{
+	FString FrameNumbers;
+
+	for (int32 Idx = 0; Idx < FRAMEHISTORYLEN; ++Idx)
+	{
+		FrameNumbers += FString::Printf(TEXT("%u"), LastRenderedFrameNumber[Idx]);
+
+		if (Idx < FRAMEHISTORYLEN - 1)
+		{
+			FrameNumbers += TEXT(", ");
+		}
+	}
+
+	return FrameNumbers;
+}
+
+FString FMediaOutputSynchronizationPolicyRivermaxHandler::FMediaSyncBarrierData::PresentedFrameBoundaryNumbersAsString() const
+{
+	FString BoundaryNumbers;
+
+	for (int32 Idx = 0; Idx < FRAMEHISTORYLEN; ++Idx)
+	{
+		BoundaryNumbers += FString::Printf(TEXT("%llu"), PresentedFrameBoundaryNumber[Idx]);
+
+		if (Idx < FRAMEHISTORYLEN - 1)
+		{
+			BoundaryNumbers += TEXT(", ");
+		}
+	}
+
+	return BoundaryNumbers;
+}
+
 void FMediaOutputSynchronizationPolicyRivermaxHandler::HandleBarrierSync(FGenericBarrierSynchronizationDelegateData& BarrierSyncData)
 {
 	TRACE_CPUPROFILER_EVENT_SCOPE(RmaxSync::BarrierSync);
@@ -257,9 +326,8 @@ void FMediaOutputSynchronizationPolicyRivermaxHandler::HandleBarrierSync(FGeneri
 		return;
 	}
 
-	FString FirstFrameNode;
-	FMediaSyncBarrierData FirstNodeData;
-	bool bFirstNode = true;
+	FString FirstNodeName;
+	const FMediaSyncBarrierData* FirstNodeData = nullptr;
 	bool bSelfRepairRequired = false;
 
 	for (const TPair<FString, TArray<uint8>>& NodePresentedFrame : BarrierSyncData.RequestData)
@@ -271,24 +339,23 @@ void FMediaOutputSynchronizationPolicyRivermaxHandler::HandleBarrierSync(FGeneri
 		// We expect all nodes to enter the barrier after presenting the SAME frame at the SAME frame number.
 		// If a node enters the barrier a frame late on the other, self repair will be triggered.
 		// Samething goes if all nodes didn't present the same frame
-		if (bFirstNode)
+		if (!FirstNodeData)
 		{
-			bFirstNode = false;
-			FirstFrameNode = NodePresentedFrame.Key;
-			FirstNodeData = *NodeData;
+			FirstNodeName = NodePresentedFrame.Key;
+			FirstNodeData = NodeData;
 		}
-		else if((FirstNodeData.LastRenderedFrameNumber != NodeData->LastRenderedFrameNumber)
-			 || (FirstNodeData.PresentedFrameBoundaryNumber != NodeData->PresentedFrameBoundaryNumber))
+		else if(FirstNodeData->HasConfirmedDesync(*NodeData))
 		{
-			UE_LOG(LogRivermaxSync, Warning, TEXT("Desync detected: Node '%s' presented frame '%u' at frame boundary '%llu', but node '%s' presented frame '%u' at frame boundary '%llu'")
-			, *FirstFrameNode
-			, FirstNodeData.LastRenderedFrameNumber
-			, FirstNodeData.PresentedFrameBoundaryNumber
-			, *NodePresentedFrame.Key
-			, NodeData->LastRenderedFrameNumber
-			, NodeData->PresentedFrameBoundaryNumber);
-
 			bSelfRepairRequired = true;
+
+			UE_LOG(LogRivermaxSync, Warning, TEXT("Desync detected: Node '%s' presented frames (%s) at boundaries (%s), but node '%s' presented frames (%s) at boundaries (%s)"),
+				*FirstNodeName, 
+				*FirstNodeData->LastRenderedFrameNumbersAsString(),
+				*FirstNodeData->PresentedFrameBoundaryNumbersAsString(),
+				*NodePresentedFrame.Key, 
+				*NodeData->LastRenderedFrameNumbersAsString(),
+				*NodeData->PresentedFrameBoundaryNumbersAsString());
+
 			break;
 		}
 	}
@@ -303,9 +370,14 @@ void FMediaOutputSynchronizationPolicyRivermaxHandler::HandleBarrierSync(FGeneri
 		const float SleepTime = TimeLeftSeconds + OffsetTimeSeconds;
 		FPlatformProcess::SleepNoStats(SleepTime);
 	}
-	else
+	else if (FirstNodeData)
 	{
-		UE_LOG(LogRivermaxSync, VeryVerbose, TEXT("'%s': Cluster synchronized. Nodes presented frame %u at frame boundary %llu"), *GetMediaDeviceId(), FirstNodeData.LastRenderedFrameNumber, FirstNodeData.PresentedFrameBoundaryNumber);
+		UE_LOG(LogRivermaxSync, VeryVerbose, TEXT("'%s': Cluster likely synchronized (no confirmed desync). Node '%s' presented frame %u at frame boundary %llu"), 
+			*GetMediaDeviceId(),
+			*FirstNodeName,
+			FirstNodeData->LastRenderedFrameNumber[0],
+			FirstNodeData->PresentedFrameBoundaryNumber[0]
+		);
 	}
 }
 
@@ -318,8 +390,10 @@ bool FMediaOutputSynchronizationPolicyRivermaxHandler::ValidateNodesFrameTime(co
 		check(NodePresentedFrame.Value.Num() == sizeof(FMediaSyncBarrierData));
 
 		const FMediaSyncBarrierData* const NodeData = reinterpret_cast<const FMediaSyncBarrierData* const>(NodePresentedFrame.Value.GetData());
-		MinFrameTime = FMath::Min(MinFrameTime, NodeData->PresentedFrameBoundaryNumber);
-		MaxFrameTime = FMath::Max(MaxFrameTime, NodeData->PresentedFrameBoundaryNumber);
+
+		// We can just use the most recent ([0]) frame presentation data to compare how close or far the nodes are from each other.
+		MinFrameTime = FMath::Min(MinFrameTime, NodeData->PresentedFrameBoundaryNumber[0]);
+		MaxFrameTime = FMath::Max(MaxFrameTime, NodeData->PresentedFrameBoundaryNumber[0]);
 	}
 
 	const uint64 MaxRange = UE::RivermaxSync::CVarRivermaxSyncMaxFrameTimeRange.GetValueOnAnyThread();
diff --git a/Engine/Plugins/VirtualProduction/Rivermax/RivermaxSync/Source/RivermaxSync/Public/MediaOutputSynchronizationPolicyRivermax.h b/Engine/Plugins/VirtualProduction/Rivermax/RivermaxSync/Source/RivermaxSync/Public/MediaOutputSynchronizationPolicyRivermax.h
--- a/Engine/Plugins/VirtualProduction/Rivermax/RivermaxSync/Source/RivermaxSync/Public/MediaOutputSynchronizationPolicyRivermax.h
+++ b/Engine/Plugins/VirtualProduction/Rivermax/RivermaxSync/Source/RivermaxSync/Public/MediaOutputSynchronizationPolicyRivermax.h
@@ -51,22 +51,57 @@ protected:
 	{
 		FMediaSyncBarrierData()
 		{
+			Reset();
 		}
 
-		/** For now, we build it directly based on presentation info from the stream */
-		FMediaSyncBarrierData(const UE::RivermaxCore::FPresentedFrameInfo& FrameInfo)
-			: PresentedFrameBoundaryNumber(FrameInfo.PresentedFrameBoundaryNumber)
-			, LastRenderedFrameNumber(FrameInfo.RenderedFrameNumber)
+		/** Reset the data to default values */
+		void Reset()
 		{
+			for (int32 Idx = 0; Idx < FRAMEHISTORYLEN; ++Idx)
+			{
+				PresentedFrameBoundaryNumber[Idx] = 0;
+				LastRenderedFrameNumber[Idx] = 0;
+			}
 		}
 
+		/** Insert the given frame information into the recorded presentation history */
+		void InsertFrameInfo(const UE::RivermaxCore::FPresentedFrameInfo& FrameInfo)
+		{
+			// Shift existing history entries
+			for (int32 Idx = FRAMEHISTORYLEN - 1; Idx > 0; --Idx)
+			{
+				PresentedFrameBoundaryNumber[Idx] = PresentedFrameBoundaryNumber[Idx - 1];
+				LastRenderedFrameNumber[Idx] = LastRenderedFrameNumber[Idx - 1];
+			}
+
+			// Insert new frame info at the beginning
+			PresentedFrameBoundaryNumber[0] = FrameInfo.PresentedFrameBoundaryNumber;
+			LastRenderedFrameNumber[0] = FrameInfo.RenderedFrameNumber;
+		}
+
+		/** Rendered frames as comma separated string */
+		FString LastRenderedFrameNumbersAsString() const;
+
+		/** Presented frame boundaries as comma separated string */
+		FString PresentedFrameBoundaryNumbersAsString() const;
+
+		/** Returns true if the frame presentation history indicates a desynced state */
+		bool HasConfirmedDesync(const FMediaSyncBarrierData& OtherBarrierData) const;
+
+		/** How many frames to include in the history */
+		static constexpr int32 FRAMEHISTORYLEN = 2;
+
 		/** Frame boundary number at which the last frame was presented.  */
-		uint64 PresentedFrameBoundaryNumber = 0;
+		uint64 PresentedFrameBoundaryNumber[FRAMEHISTORYLEN];
 		
 		/** Last engine frame number that was presented */
-		uint32 LastRenderedFrameNumber = 0;
+		uint32 LastRenderedFrameNumber[FRAMEHISTORYLEN];
+
 	};
 
+	/** Holds data provided to server by this node when joining the barrier */
+	struct FMediaSyncBarrierData BarrierDataStruct;
+
 	/** Synchronization margin (ms) */
 	float MarginMs = 5.0f;
 
