Rivermax: Report when the same frame was presented on different PTP frame boundaries.

#rnx

diff --git a/Engine/Plugins/VirtualProduction/Rivermax/RivermaxSync/Source/RivermaxSync/Private/MediaOutputSynchronizationPolicyRivermax.cpp b/Engine/Plugins/VirtualProduction/Rivermax/RivermaxSync/Source/RivermaxSync/Private/MediaOutputSynchronizationPolicyRivermax.cpp
--- a/Engine/Plugins/VirtualProduction/Rivermax/RivermaxSync/Source/RivermaxSync/Private/MediaOutputSynchronizationPolicyRivermax.cpp
+++ b/Engine/Plugins/VirtualProduction/Rivermax/RivermaxSync/Source/RivermaxSync/Private/MediaOutputSynchronizationPolicyRivermax.cpp
@@ -237,8 +237,10 @@ bool FMediaOutputSynchronizationPolicyRivermaxHandler::InitializeBarrier(const F
 }
 
 
-bool FMediaOutputSynchronizationPolicyRivermaxHandler::FMediaSyncBarrierData::HasConfirmedDesync(const FMediaSyncBarrierData& OtherBarrierData) const
+bool FMediaOutputSynchronizationPolicyRivermaxHandler::FMediaSyncBarrierData::HasConfirmedDesync(const FMediaSyncBarrierData& OtherBarrierData, uint64& OutVsyncDelta) const
 {
+	OutVsyncDelta = 0;
+
 	for (int32 FirstNodeIdx = 0; FirstNodeIdx < FMediaSyncBarrierData::FRAMEHISTORYLEN; ++FirstNodeIdx)
 	{
 		const uint32 FirstNodeFrameNumber = LastRenderedFrameNumber[FirstNodeIdx];
@@ -252,6 +254,16 @@ bool FMediaOutputSynchronizationPolicyRivermaxHandler::FMediaSyncBarrierData::Ha
 			const bool bSameFrame = (FirstNodeFrameNumber == OtherNodeFrameNumber);
 			const bool bSameVsync = (FirstNodeVsyncBoundary == OtherNodeVsyncBoundary);
 
+			// Keep track of the maximum Vsync delta of equal frames.
+			if (bSameFrame && !bSameVsync)
+			{
+				const uint64 VsyncDelta = FirstNodeVsyncBoundary > OtherNodeVsyncBoundary ?
+					FirstNodeVsyncBoundary - OtherNodeVsyncBoundary :
+					OtherNodeVsyncBoundary - FirstNodeVsyncBoundary;
+
+				OutVsyncDelta = FMath::Max(OutVsyncDelta, VsyncDelta);
+			}
+
 			// If they agree on a recent frame, consider them in sync
 			if (bSameFrame && bSameVsync)
 			{
@@ -343,20 +355,36 @@ void FMediaOutputSynchronizationPolicyRivermaxHandler::HandleBarrierSync(FGeneri
 		{
 			FirstNodeName = NodePresentedFrame.Key;
 			FirstNodeData = NodeData;
+			continue;
 		}
-		else if(FirstNodeData->HasConfirmedDesync(*NodeData))
+
+		uint64 VsyncDelta = 0;
+
+		if (FirstNodeData->HasConfirmedDesync(*NodeData, VsyncDelta))
 		{
 			bSelfRepairRequired = true;
 
-			UE_LOG(LogRivermaxSync, Warning, TEXT("Desync detected: Node '%s' presented frames (%s) at boundaries (%s), but node '%s' presented frames (%s) at boundaries (%s)"),
-				*FirstNodeName, 
+			UE_LOG(LogRivermaxSync, Warning,
+				TEXT("Desync detected: Node '%s' presented frames (%s) at boundaries (%s), but node '%s' presented frames (%s) at boundaries (%s)"),
+				*FirstNodeName,
 				*FirstNodeData->LastRenderedFrameNumbersAsString(),
 				*FirstNodeData->PresentedFrameBoundaryNumbersAsString(),
-				*NodePresentedFrame.Key, 
+				*NodePresentedFrame.Key,
 				*NodeData->LastRenderedFrameNumbersAsString(),
 				*NodeData->PresentedFrameBoundaryNumbersAsString());
 
-			break;
+			// We do not break the loop, in order to log all timing issues detected in the current frame.
+		}
+		else if (VsyncDelta > 0)
+		{
+			UE_LOG(LogRivermaxSync, Warning,
+				TEXT("Frames not presented at the same PTP frame boundary: Node '%s' presented frames (%s) at boundaries (%s), but node '%s' presented frames (%s) at boundaries (%s)"),
+				*FirstNodeName,
+				*FirstNodeData->LastRenderedFrameNumbersAsString(),
+				*FirstNodeData->PresentedFrameBoundaryNumbersAsString(),
+				*NodePresentedFrame.Key,
+				*NodeData->LastRenderedFrameNumbersAsString(),
+				*NodeData->PresentedFrameBoundaryNumbersAsString());
 		}
 	}
 
diff --git a/Engine/Plugins/VirtualProduction/Rivermax/RivermaxSync/Source/RivermaxSync/Public/MediaOutputSynchronizationPolicyRivermax.h b/Engine/Plugins/VirtualProduction/Rivermax/RivermaxSync/Source/RivermaxSync/Public/MediaOutputSynchronizationPolicyRivermax.h
--- a/Engine/Plugins/VirtualProduction/Rivermax/RivermaxSync/Source/RivermaxSync/Public/MediaOutputSynchronizationPolicyRivermax.h
+++ b/Engine/Plugins/VirtualProduction/Rivermax/RivermaxSync/Source/RivermaxSync/Public/MediaOutputSynchronizationPolicyRivermax.h
@@ -85,8 +85,20 @@ protected:
 		/** Presented frame boundaries as comma separated string */
 		FString PresentedFrameBoundaryNumbersAsString() const;
 
-		/** Returns true if the frame presentation history indicates a desynced state */
-		bool HasConfirmedDesync(const FMediaSyncBarrierData& OtherBarrierData) const;
+		/** 
+		 * Returns true if the frame presentation history indicates a desynced state. 
+		 * 
+		 * @param OtherBarrierData
+		 *     Barrier data of the node we're comparing with.
+		 * 
+		 * @param OutVsyncDelta
+		 *     When the same frame is presented at different Vsync frame boundaries, this parameter
+		 *     contains the delta between them. Useful in detecting large PTP differences between nodes.
+		 * 
+		 * @return
+		 *     True if different frames were presented at the same Vsync frame boundaries.
+		 */
+		bool HasConfirmedDesync(const FMediaSyncBarrierData& OtherBarrierData, uint64& OutVsyncDelta) const;
 
 		/** How many frames to include in the history */
 		static constexpr int32 FRAMEHISTORYLEN = 2;
