nDisplay and Rivermax: The p-node now sends an event every 120 frames (cvar to set frequency) with a list of any detected PTP video frame mismatches.

The goal is that the operator is able to detect the nodes that have PTP misconfiguration by looking at the Stage Monitor events, and also get an idea of the extent of the  PTP mismatch.

PTP mismatches are made with respect to the primary node, or if the primary node is not in the barrier, then the first node in alphabetical order.

Removed an initial check of PTP mismatch that would disable or enable the ptp repair. This is no longer necessary since the repair is only triggered when there are nodes with likely repairable mismatches.

Use Rivermax.Sync.Ptp.UnsyncFramesPerReport to control how often to send the events (-1 to disable them).
When enabled, the event is always generated, including when all nodes are in PTP sync, to give positive confirmation to the operator that the cluster is operating as expected.

#rn nDisplay node PTP mistmatches are now reported as Stage Monitor events.

diff --git a/Engine/Plugins/VirtualProduction/Rivermax/RivermaxSync/Source/RivermaxSync/Private/MediaOutputSynchronizationPolicyRivermax.cpp b/Engine/Plugins/VirtualProduction/Rivermax/RivermaxSync/Source/RivermaxSync/Private/MediaOutputSynchronizationPolicyRivermax.cpp
--- a/Engine/Plugins/VirtualProduction/Rivermax/RivermaxSync/Source/RivermaxSync/Private/MediaOutputSynchronizationPolicyRivermax.cpp
+++ b/Engine/Plugins/VirtualProduction/Rivermax/RivermaxSync/Source/RivermaxSync/Private/MediaOutputSynchronizationPolicyRivermax.cpp
@@ -2,7 +2,9 @@
 
 #include "MediaOutputSynchronizationPolicyRivermax.h"
 
+#include "Async/Async.h"
 #include "Cluster/IDisplayClusterClusterManager.h"
+#include "Config/IDisplayClusterConfigManager.h"
 #include "HAL/IConsoleManager.h"
 #include "IDisplayCluster.h"
 #include "IRivermaxCoreModule.h"
@@ -12,7 +14,6 @@
 #include "RivermaxSyncLog.h"
 
 
-
 namespace UE::RivermaxSync
 {
 	static TAutoConsoleVariable<float> CVarRivermaxSyncWakeupOffset(
@@ -25,9 +26,11 @@ namespace UE::RivermaxSync
 		TEXT("Whether to use exchanged data in the synchronization barrier to detect desynchronized state and act on it to self repair"),
 		ECVF_Default);
 
-	static TAutoConsoleVariable<int32> CVarRivermaxSyncMaxFrameTimeRange(
-		TEXT("Rivermax.Sync.MaxFrameTimeRange"), 8,
-		TEXT("Maximum number of frame last presented frame from nodes can be from each other when they join the barrier to enable self repair"),
+	static TAutoConsoleVariable<int32> CVarRivermaxPtpUnsyncFramesPerReport(
+		TEXT("Rivermax.Sync.Ptp.UnsyncFramesPerReport"), 120,
+		TEXT("When there are PTP mismatches in the cluster, Stage Monitor events are issued.\n"
+			 "PTP mismatches are a stable condition and this cvar controls how often to send the events.\n"
+			 "Use -1 (or any negative number) to disable these reports."),
 		ECVF_Default);
 
 	static bool GbTriggerRandomDesync = false;
@@ -35,8 +38,10 @@ namespace UE::RivermaxSync
 		TEXT("Rivermax.Sync.ForceDesync")
 		, UE::RivermaxSync::GbTriggerRandomDesync
 		, TEXT("After barrier synchronization, trigger random stall."), ECVF_Cheat);
+
 }
 
+
 FMediaOutputSynchronizationPolicyRivermaxHandler::FMediaOutputSynchronizationPolicyRivermaxHandler(UMediaOutputSynchronizationPolicyRivermax* InPolicyObject)
 	: Super(InPolicyObject)
 	, MarginMs(InPolicyObject->MarginMs)
@@ -237,7 +242,7 @@ bool FMediaOutputSynchronizationPolicyRivermaxHandler::InitializeBarrier(const F
 }
 
 
-bool FMediaOutputSynchronizationPolicyRivermaxHandler::FMediaSyncBarrierData::HasConfirmedDesync(const FMediaSyncBarrierData& OtherBarrierData, uint64& OutVsyncDelta) const
+bool FMediaOutputSynchronizationPolicyRivermaxHandler::FMediaSyncBarrierData::HasConfirmedDesync(const FMediaSyncBarrierData& OtherBarrierData, int64& OutVsyncDelta) const
 {
 	OutVsyncDelta = 0;
 
@@ -257,11 +262,13 @@ bool FMediaOutputSynchronizationPolicyRivermaxHandler::FMediaSyncBarrierData::Ha
 			// Keep track of the maximum Vsync delta of equal frames.
 			if (bSameFrame && !bSameVsync)
 			{
-				const uint64 VsyncDelta = FirstNodeVsyncBoundary > OtherNodeVsyncBoundary ?
-					FirstNodeVsyncBoundary - OtherNodeVsyncBoundary :
-					OtherNodeVsyncBoundary - FirstNodeVsyncBoundary;
+				// Positive means that the node has a PTP frame number larger than the base we're comparing with.
+				const int64 VsyncDelta = static_cast<int64>(OtherNodeVsyncBoundary) - static_cast<int64>(FirstNodeVsyncBoundary);
 
-				OutVsyncDelta = FMath::Max(OutVsyncDelta, VsyncDelta);
+				if (FMath::Abs(VsyncDelta) > FMath::Abs(OutVsyncDelta))
+				{
+					OutVsyncDelta = VsyncDelta;
+				}
 			}
 
 			// If they agree on a recent frame, consider them in sync
@@ -317,134 +324,237 @@ FString FMediaOutputSynchronizationPolicyRivermaxHandler::FMediaSyncBarrierData:
 	return BoundaryNumbers;
 }
 
-void FMediaOutputSynchronizationPolicyRivermaxHandler::HandleBarrierSync(FGenericBarrierSynchronizationDelegateData& BarrierSyncData)
+bool FMediaOutputSynchronizationPolicyRivermaxHandler::PickPtpBaseNodeAndData(
+	const FGenericBarrierSynchronizationDelegateData& BarrierSyncData, 
+	const FMediaSyncBarrierData*& OutPtpBaseNodeData,
+	FString& OutBaseNodeId
+) const
 {
-	TRACE_CPUPROFILER_EVENT_SCOPE(RmaxSync::BarrierSync);
+	// Identify the name of the primary node, which will be the default for PTP mismatch measurements.
 
-	if (bHasVerifiedClocks == false)
+	OutPtpBaseNodeData = nullptr;
+
+	// We need at least one node.
+	if (BarrierSyncData.RequestData.Num() <= 0)
 	{
-		bHasVerifiedClocks = true;
-		bCanUseSelfRepair = ValidateNodesFrameTime(BarrierSyncData.RequestData);
+		return false;
 	}
 
-	if (!bCanUseSelfRepair || UE::RivermaxSync::CVarRivermaxSyncEnableSelfRepair.GetValueOnAnyThread() == false)
+	IDisplayClusterConfigManager* ConfigMgr = IDisplayCluster::Get().GetConfigMgr();
+
+	if (!ConfigMgr)
 	{
-		return;
+		return false;
 	}
 
+	OutBaseNodeId = ConfigMgr->GetPrimaryNodeId();
+	const TArray<uint8>* PrimaryDataRaw = BarrierSyncData.RequestData.Find(OutBaseNodeId);
+
+	if (PrimaryDataRaw)
+	{
+		check(PrimaryDataRaw->Num() == sizeof(FMediaSyncBarrierData));
+		OutPtpBaseNodeData = reinterpret_cast<const FMediaSyncBarrierData*>(PrimaryDataRaw->GetData());
+
+		return true;
+	}
+
+	// If the primary node isn't in the barrier, then we pick the first one in a sorted list of node ids.
+	// We don't need to actually sort the list, just find the first one, which is O(n) instead of O(n log n).
+
+	const FString* FirstKeyIfSorted = nullptr;
+	const TArray<uint8>* FirstValueIfSorted = nullptr;
+
+	for (const auto& Elem : BarrierSyncData.RequestData)
+	{
+		const FString& Key = Elem.Key;
+		const TArray<uint8>& Value = Elem.Value;
+
+		if (!FirstKeyIfSorted || Key < *FirstKeyIfSorted)
+		{
+			FirstKeyIfSorted = &Key;
+			FirstValueIfSorted = &Value;
+		}
+	}
+
+	check(FirstKeyIfSorted && FirstValueIfSorted);
+
+	OutBaseNodeId = *FirstKeyIfSorted;
+
+	check(FirstValueIfSorted->Num() == sizeof(FMediaSyncBarrierData));
+	OutPtpBaseNodeData = reinterpret_cast<const FMediaSyncBarrierData*>(FirstValueIfSorted->GetData());
+
+	return true;
+}
+
+void FMediaOutputSynchronizationPolicyRivermaxHandler::HandleBarrierSync(FGenericBarrierSynchronizationDelegateData& BarrierSyncData)
+{
+	TRACE_CPUPROFILER_EVENT_SCOPE(RmaxSync::BarrierSync);
+
+	// Nothing to do if there is no barrier data.
 	if (BarrierSyncData.RequestData.Num() <= 0)
 	{
 		UE_LOG(LogRivermaxSync, Verbose, TEXT("'%s': No data was provided by nodes for sync barrier."), *GetMediaDeviceId())
 		return;
 	}
 
-	FString FirstNodeName;
-	const FMediaSyncBarrierData* FirstNodeData = nullptr;
-	bool bSelfRepairRequired = false;
+	// Deterministically pick a node in the cluster to use as the PTP base for mismatch detections.
 
+	const FMediaSyncBarrierData* PtpBaseNodeData = nullptr;
+	FString PtpBaseNodeId;
+
+	if (!PickPtpBaseNodeAndData(BarrierSyncData, PtpBaseNodeData, PtpBaseNodeId))
+	{
+		UE_LOG(LogRivermaxSync, Warning, TEXT("Could not find a base node for ptp mismatch comparisons"));
+		return;
+	}
+
+	check(PtpBaseNodeData);
+
+	bool bSelfRepairRequired = false; // True if there is any need (and benefit) in initiating a sync repair action.
+
+	TMap<FString, int64> PtpMismatchedNodes; // Collects node ids with mismatched PTP frames.
+
+	// To avoid flooding the network and the receiver with these events, we only report ptp mismatches every few frames.
+	const int32 PtpUnsyncFramesPerReport = UE::RivermaxSync::CVarRivermaxPtpUnsyncFramesPerReport.GetValueOnAnyThread();
+	bool bShouldReportPtpMismatches = (PtpUnsyncFramesPerReport >= 0) && !(PtpBaseNodeData->LastRenderedFrameNumber[0] % PtpUnsyncFramesPerReport);
+
+	// Iterate over the nodes and detect PTP de-syncs.
 	for (const TPair<FString, TArray<uint8>>& NodePresentedFrame : BarrierSyncData.RequestData)
 	{
+		// Barrier data with unexpected sizes are a logical error.
 		check(NodePresentedFrame.Value.Num() == sizeof(FMediaSyncBarrierData));
 
+		// Get the node data in struct format.
 		const FMediaSyncBarrierData* const NodeData = reinterpret_cast<const FMediaSyncBarrierData* const>(NodePresentedFrame.Value.GetData());
-		
-		// We expect all nodes to enter the barrier after presenting the SAME frame at the SAME frame number.
-		// If a node enters the barrier a frame late on the other, self repair will be triggered.
-		// Samething goes if all nodes didn't present the same frame
-		if (!FirstNodeData)
+
+		// Skip base node comparing with itself
+		if (NodeData == PtpBaseNodeData)
 		{
-			FirstNodeName = NodePresentedFrame.Key;
-			FirstNodeData = NodeData;
 			continue;
 		}
 
-		uint64 VsyncDelta = 0;
+		// We expect all nodes to enter the barrier after presenting the SAME frame at the SAME frame number.
+		// If a node enters the barrier a frame late on the other, self repair will be triggered.
+		// Same thing goes if all nodes didn't present the same frame
 
-		if (FirstNodeData->HasConfirmedDesync(*NodeData, VsyncDelta))
+		int64 VsyncDelta = 0;
+
+		if (PtpBaseNodeData->HasConfirmedDesync(*NodeData, VsyncDelta))
 		{
 			bSelfRepairRequired = true;
 
 			UE_LOG(LogRivermaxSync, Warning,
 				TEXT("Desync detected: Node '%s' presented frames (%s) at boundaries (%s), but node '%s' presented frames (%s) at boundaries (%s)"),
-				*FirstNodeName,
-				*FirstNodeData->LastRenderedFrameNumbersAsString(),
-				*FirstNodeData->PresentedFrameBoundaryNumbersAsString(),
+				*PtpBaseNodeId,
+				*PtpBaseNodeData->LastRenderedFrameNumbersAsString(),
+				*PtpBaseNodeData->PresentedFrameBoundaryNumbersAsString(),
 				*NodePresentedFrame.Key,
 				*NodeData->LastRenderedFrameNumbersAsString(),
 				*NodeData->PresentedFrameBoundaryNumbersAsString());
 
 			// We do not break the loop, in order to log all timing issues detected in the current frame.
 		}
-		else if (VsyncDelta > 0)
+		else if (VsyncDelta != 0)
 		{
 			UE_LOG(LogRivermaxSync, Warning,
 				TEXT("Frames not presented at the same PTP frame boundary: Node '%s' presented frames (%s) at boundaries (%s), but node '%s' presented frames (%s) at boundaries (%s)"),
-				*FirstNodeName,
-				*FirstNodeData->LastRenderedFrameNumbersAsString(),
-				*FirstNodeData->PresentedFrameBoundaryNumbersAsString(),
+				*PtpBaseNodeId,
+				*PtpBaseNodeData->LastRenderedFrameNumbersAsString(),
+				*PtpBaseNodeData->PresentedFrameBoundaryNumbersAsString(),
 				*NodePresentedFrame.Key,
 				*NodeData->LastRenderedFrameNumbersAsString(),
 				*NodeData->PresentedFrameBoundaryNumbersAsString());
 		}
+
+		// Collect Vsync deltas for reporting purposes.
+		if (bShouldReportPtpMismatches && VsyncDelta)
+		{
+			PtpMismatchedNodes.Add(NodePresentedFrame.Key, VsyncDelta);
+		}
 	}
 
+	// Report to Stage Monitor the PTP mismatches.
+	if (bShouldReportPtpMismatches)
+	{
+		AsyncTask(ENamedThreads::GameThread, [PtpMismatchedNodes, PtpBaseNodeId]()
+			{
+				IStageDataProvider::SendMessage<FRivermaxClusterPtpUnsyncEvent>(
+					EStageMessageFlags::None, // Doesn't need to be reliable.
+					PtpMismatchedNodes,
+					PtpBaseNodeId
+				);
+			});
+	}
+
+	// This cvar can disable the self repair.
+	const bool bCanUseSelfRepair = UE::RivermaxSync::CVarRivermaxSyncEnableSelfRepair.GetValueOnAnyThread();
+
 	// If repair is required, we stall until we are past the next alignment point to have all scheduler present something and get closer to a synchronized state.
 	if (bSelfRepairRequired)
 	{
-		TRACE_CPUPROFILER_EVENT_SCOPE(RmaxSync::SelfRepair);
+		// Only trigger the self repair when allowed. It consists on sleeping until the next ptp vsync passes.
+		if (bCanUseSelfRepair)
+		{
+			TRACE_CPUPROFILER_EVENT_SCOPE(RmaxSync::SelfRepair);
 
-		const double TimeLeftSeconds = GetTimeBeforeNextSyncPoint();
-		const float OffsetTimeSeconds = UE::RivermaxSync::CVarRivermaxSyncWakeupOffset.GetValueOnAnyThread() * 1E-3;
-		const float SleepTime = TimeLeftSeconds + OffsetTimeSeconds;
-		FPlatformProcess::SleepNoStats(SleepTime);
+			const double TimeLeftSeconds = GetTimeBeforeNextSyncPoint();
+			const float OffsetTimeSeconds = UE::RivermaxSync::CVarRivermaxSyncWakeupOffset.GetValueOnAnyThread() * 1E-3;
+			const float SleepTime = TimeLeftSeconds + OffsetTimeSeconds;
+			FPlatformProcess::SleepNoStats(SleepTime);
+		}
 	}
-	else if (FirstNodeData)
+	else
 	{
-		UE_LOG(LogRivermaxSync, VeryVerbose, TEXT("'%s': Cluster likely synchronized (no confirmed desync). Node '%s' presented frame %u at frame boundary %llu"), 
+		// VeryVerbose log of the ptp frame presentation values when everything is in ptp sync.
+
+		UE_LOG(LogRivermaxSync, VeryVerbose, TEXT("'%s': Cluster likely synchronized (no confirmed desync). ptp base node '%s' presented frame %u at frame boundary %llu"), 
 			*GetMediaDeviceId(),
-			*FirstNodeName,
-			FirstNodeData->LastRenderedFrameNumber[0],
-			FirstNodeData->PresentedFrameBoundaryNumber[0]
+			*PtpBaseNodeId,
+			PtpBaseNodeData->LastRenderedFrameNumber[0],
+			PtpBaseNodeData->PresentedFrameBoundaryNumber[0]
 		);
 	}
 }
 
-bool FMediaOutputSynchronizationPolicyRivermaxHandler::ValidateNodesFrameTime(const TMap<FString, TArray<uint8>>& NodeRequestData) const
+TSharedPtr<IDisplayClusterMediaOutputSynchronizationPolicyHandler> UMediaOutputSynchronizationPolicyRivermax::GetHandler()
 {
-	uint64 MinFrameTime = TNumericLimits<uint64>::Max();
-	uint64 MaxFrameTime = TNumericLimits<uint64>::Min();
-	for (const TPair<FString, TArray<uint8>>& NodePresentedFrame : NodeRequestData)
+	if (!Handler)
 	{
-		check(NodePresentedFrame.Value.Num() == sizeof(FMediaSyncBarrierData));
-
-		const FMediaSyncBarrierData* const NodeData = reinterpret_cast<const FMediaSyncBarrierData* const>(NodePresentedFrame.Value.GetData());
-
-		// We can just use the most recent ([0]) frame presentation data to compare how close or far the nodes are from each other.
-		MinFrameTime = FMath::Min(MinFrameTime, NodeData->PresentedFrameBoundaryNumber[0]);
-		MaxFrameTime = FMath::Max(MaxFrameTime, NodeData->PresentedFrameBoundaryNumber[0]);
+		Handler = MakeShared<FMediaOutputSynchronizationPolicyRivermaxHandler>(this);
 	}
 
-	const uint64 MaxRange = UE::RivermaxSync::CVarRivermaxSyncMaxFrameTimeRange.GetValueOnAnyThread();
-	const uint64 DetectedRange = MaxFrameTime - MinFrameTime;
-	if (DetectedRange > MaxRange)
-	{	
-		UE_LOG(LogRivermaxSync, Warning, TEXT("Self repair can't be enabled. Frame time range (%llu) across cluster too large. Verify PTP clocks to be identical on each node."), DetectedRange);
-		return false;
-	}
-	else
-	{
-		UE_LOG(LogRivermaxSync, Log, TEXT("Self repair can be used. Frame time range (%llu) across cluster points to a common time reference."), DetectedRange);
-		return true;
-	}
+	return Handler;
 }
 
-TSharedPtr<IDisplayClusterMediaOutputSynchronizationPolicyHandler> UMediaOutputSynchronizationPolicyRivermax::GetHandler()
+
+FString FRivermaxClusterPtpUnsyncEvent::ToString() const
 {
-	if (!Handler)
+	if (NodePtpFrameDeltas.Num() > 0)
 	{
-		Handler = MakeShared<FMediaOutputSynchronizationPolicyRivermaxHandler>(this);
+		FString Result = FString::Printf(TEXT("PTP video frame mismatches compared to PTP base node '%s' on nodes: "), *PtpBaseNodeId);
+
+		// Sort the keys to produce a deterministic string output, which will help readability of repeated events.
+
+		TArray<FString> SortedKeys;
+		NodePtpFrameDeltas.GetKeys(SortedKeys);
+		SortedKeys.Sort();
+
+		// Iterate over sorted keys for deterministic output
+		for (int32 NodeIdx = 0; NodeIdx < SortedKeys.Num(); ++NodeIdx)
+		{
+			const FString& Key = SortedKeys[NodeIdx];
+			int64 PtpFrameDelta = NodePtpFrameDeltas[Key];
+
+			Result += FString::Printf(TEXT("%s(%lld)"), *Key, PtpFrameDelta);
+
+			if (NodeIdx < SortedKeys.Num() - 1)
+			{
+				Result += TEXT(", ");
+			}
+		}
+
+		return Result;
 	}
 
-	return Handler;
+	return FString(TEXT("All nodes are in PTP sync."));
 }
-
diff --git a/Engine/Plugins/VirtualProduction/Rivermax/RivermaxSync/Source/RivermaxSync/Public/MediaOutputSynchronizationPolicyRivermax.h b/Engine/Plugins/VirtualProduction/Rivermax/RivermaxSync/Source/RivermaxSync/Public/MediaOutputSynchronizationPolicyRivermax.h
--- a/Engine/Plugins/VirtualProduction/Rivermax/RivermaxSync/Source/RivermaxSync/Public/MediaOutputSynchronizationPolicyRivermax.h
+++ b/Engine/Plugins/VirtualProduction/Rivermax/RivermaxSync/Source/RivermaxSync/Public/MediaOutputSynchronizationPolicyRivermax.h
@@ -5,6 +5,7 @@
 #include "Synchronization/DisplayClusterMediaOutputSynchronizationPolicyEthernetBarrierBase.h"
 
 #include "IRivermaxOutputStream.h"
+#include "IStageDataProvider.h"
 
 #include "MediaOutputSynchronizationPolicyRivermax.generated.h"
 
@@ -31,19 +32,6 @@ public:
 	/** Returns true if specified media capture type can be synchonized by the policy implementation */
 	virtual bool IsCaptureTypeSupported(UMediaCapture* MediaCapture) const override;
 
-protected:
-	/** Initializes dynamic barrier on the primary node. */
-	virtual bool InitializeBarrier(const FString& SyncInstanceId) override;
-
-	/** Barrier callback containing data from each node to detect if cluster is out of sync. */
-	void HandleBarrierSync(FGenericBarrierSynchronizationDelegateData& BarrierSyncData);
-
-	/** Returns true if frame time from each node are close to each other (limit controlled by cvar) */
-	bool ValidateNodesFrameTime(const TMap<FString, TArray<uint8>>& NodeRequestData) const;
-
-	/** Returns amount of time before next synchronization point. */
-	double GetTimeBeforeNextSyncPoint() const;
-
 protected:
 
 	/** Holds data provided to server by each node when joining the barrier */
@@ -85,32 +73,57 @@ protected:
 		/** Presented frame boundaries as comma separated string */
 		FString PresentedFrameBoundaryNumbersAsString() const;
 
-		/** 
-		 * Returns true if the frame presentation history indicates a desynced state. 
-		 * 
+		/**
+		 * Returns true if the frame presentation history indicates a desynced state.
+		 *
 		 * @param OtherBarrierData
 		 *     Barrier data of the node we're comparing with.
-		 * 
+		 *
 		 * @param OutVsyncDelta
 		 *     When the same frame is presented at different Vsync frame boundaries, this parameter
 		 *     contains the delta between them. Useful in detecting large PTP differences between nodes.
-		 * 
+		 *
 		 * @return
 		 *     True if different frames were presented at the same Vsync frame boundaries.
 		 */
-		bool HasConfirmedDesync(const FMediaSyncBarrierData& OtherBarrierData, uint64& OutVsyncDelta) const;
+		bool HasConfirmedDesync(const FMediaSyncBarrierData& OtherBarrierData, int64& OutVsyncDelta) const;
 
 		/** How many frames to include in the history */
 		static constexpr int32 FRAMEHISTORYLEN = 2;
 
 		/** Frame boundary number at which the last frame was presented.  */
 		uint64 PresentedFrameBoundaryNumber[FRAMEHISTORYLEN];
-		
+
 		/** Last engine frame number that was presented */
 		uint32 LastRenderedFrameNumber[FRAMEHISTORYLEN];
-
 	};
 
+protected:
+
+	/** Initializes dynamic barrier on the primary node. */
+	virtual bool InitializeBarrier(const FString& SyncInstanceId) override;
+
+	/** Barrier callback containing data from each node to detect if cluster is out of sync. */
+	void HandleBarrierSync(FGenericBarrierSynchronizationDelegateData& BarrierSyncData);
+
+	/** Returns amount of time before next synchronization point. */
+	double GetTimeBeforeNextSyncPoint() const;
+
+	/** 
+	 * Deterministically picks a node to base ptp offsets on other nodes from.
+	 * 
+	 * @param BarrierSyncData the barrier data from all the clients.
+	 * @param OutPtpBaseNodeData Will be populated with a pointer to the barrier sync data of the selected ptp base node.
+	 * @param OutPtpBaseNodeId For convenience, this returns the name of the ptp base node.
+	 */
+	bool PickPtpBaseNodeAndData(
+		const FGenericBarrierSynchronizationDelegateData& BarrierSyncData,
+		const FMediaSyncBarrierData*& OutPtpBaseNodeData,
+		FString& OutPtpBaseNodeId
+	) const;
+
+protected:
+
 	/** Holds data provided to server by this node when joining the barrier */
 	struct FMediaSyncBarrierData BarrierDataStruct;
 
@@ -119,12 +132,6 @@ protected:
 
 	/** Memory buffer used to contain data exchanged in the barrier */
 	TArray<uint8> BarrierData;
-
-	/** On first barrier sync, we will verify frame time range to detect if clocks are too far appart. Even when desynchronized, frame time shouldn't be further than 1-2 frame away */
-	bool bCanUseSelfRepair = true;
-
-	/** Verify clock between nodes only once */
-	bool bHasVerifiedClocks = false;
 };
 
 
@@ -149,3 +156,37 @@ public:
 	UPROPERTY(EditAnywhere, BlueprintReadWrite, Category = "Settings", meta = (DisplayName = "Margin (ms)", ClampMin = "1", ClampMax = "20", UIMin = "1", UIMax = "20"))
 	float MarginMs = 5.0f;
 };
+
+/**
+ * Stage Monitor event to report nodes that are out of PTP sync with respect to a given PTP base node.
+ */
+USTRUCT()
+struct FRivermaxClusterPtpUnsyncEvent : public FStageProviderEventMessage
+{
+	GENERATED_BODY()
+
+public:
+	FRivermaxClusterPtpUnsyncEvent() = default;
+
+	FRivermaxClusterPtpUnsyncEvent(const TMap<FString, int64>& InNodePtpFrameDeltas, const FString& InPtpBaseNodeId)
+		: NodePtpFrameDeltas(InNodePtpFrameDeltas)
+		, PtpBaseNodeId(InPtpBaseNodeId)
+	{}
+
+public:
+
+	/** Nodes with PTP video frame mismatches compared to the ptp base node id. */
+	UPROPERTY(VisibleAnywhere, Category = "PtpSync")
+	TMap<FString, int64> NodePtpFrameDeltas;
+
+	/** Id of the base node the PTP delta video frames are compared with */
+	UPROPERTY(VisibleAnywhere, Category = "PtpSync")
+	FString PtpBaseNodeId;
+
+public:
+
+	//~ Begin FStageDataBaseMessage
+	virtual FString ToString() const override;
+	//~ End FStageDataBaseMessage
+};
+
diff --git a/Engine/Plugins/VirtualProduction/Rivermax/RivermaxSync/Source/RivermaxSync/RivermaxSync.Build.cs b/Engine/Plugins/VirtualProduction/Rivermax/RivermaxSync/Source/RivermaxSync/RivermaxSync.Build.cs
--- a/Engine/Plugins/VirtualProduction/Rivermax/RivermaxSync/Source/RivermaxSync/RivermaxSync.Build.cs
+++ b/Engine/Plugins/VirtualProduction/Rivermax/RivermaxSync/Source/RivermaxSync/RivermaxSync.Build.cs
@@ -8,15 +8,17 @@ public class RivermaxSync : ModuleRules
 	{
 		PublicDependencyModuleNames.AddRange(
 			new string[] {
+				"DisplayCluster",
 				"DisplayClusterMedia",
-				"RivermaxCore"
+				"RivermaxCore",
+				"StageDataCore",
 			});
 
 		PrivateDependencyModuleNames.AddRange(
 			new string[] {
 				"Core",
 				"CoreUObject",
-				"RivermaxMedia"
+				"RivermaxMedia",
 			});
 
 		if (Target.bBuildEditor == true)
