Fixed printf format specifier mismatch.

diff --git a/Engine/Plugins/VirtualProduction/Rivermax/RivermaxSync/Source/RivermaxSync/Private/MediaOutputSynchronizationPolicyRivermax.cpp b/Engine/Plugins/VirtualProduction/Rivermax/RivermaxSync/Source/RivermaxSync/Private/MediaOutputSynchronizationPolicyRivermax.cpp
--- a/Engine/Plugins/VirtualProduction/Rivermax/RivermaxSync/Source/RivermaxSync/Private/MediaOutputSynchronizationPolicyRivermax.cpp
+++ b/Engine/Plugins/VirtualProduction/Rivermax/RivermaxSync/Source/RivermaxSync/Private/MediaOutputSynchronizationPolicyRivermax.cpp
@@ -142,7 +142,7 @@ void FMediaOutputSynchronizationPolicyRivermaxHandler::Synchronize()
 
 		// Fill the memory to be exchanged by nodes in the barrier.
 		BarrierDataStruct.InsertFrameInfo(FrameInfo);
-		UE_LOG(LogRivermaxSync, VeryVerbose, TEXT("'%s' Entering with %u"), *GetMediaDeviceId(), BarrierDataStruct.LastRenderedFrameNumber);
+		UE_LOG(LogRivermaxSync, VeryVerbose, TEXT("'%s' Entering with %u"), *GetMediaDeviceId(), BarrierDataStruct.LastRenderedFrameNumber[0]);
 		FMemory::Memcpy(BarrierData.GetData(), &BarrierDataStruct, sizeof(FMediaSyncBarrierData));
 
 		// We don't use response data for now
