Rivermax: Fix a possible mod 0 depending on cvar value.

#rnx

diff --git a/Engine/Plugins/VirtualProduction/Rivermax/RivermaxSync/Source/RivermaxSync/Private/MediaOutputSynchronizationPolicyRivermax.cpp b/Engine/Plugins/VirtualProduction/Rivermax/RivermaxSync/Source/RivermaxSync/Private/MediaOutputSynchronizationPolicyRivermax.cpp
--- a/Engine/Plugins/VirtualProduction/Rivermax/RivermaxSync/Source/RivermaxSync/Private/MediaOutputSynchronizationPolicyRivermax.cpp
+++ b/Engine/Plugins/VirtualProduction/Rivermax/RivermaxSync/Source/RivermaxSync/Private/MediaOutputSynchronizationPolicyRivermax.cpp
@@ -431,7 +431,7 @@ void FMediaOutputSynchronizationPolicyRivermaxHandler::HandleBarrierSync(FGeneri
 
 	// To avoid flooding the network and the receiver with these events, we only report ptp mismatches every few frames.
 	const int32 PtpUnsyncFramesPerReport = UE::RivermaxSync::CVarRivermaxPtpUnsyncFramesPerReport.GetValueOnAnyThread();
-	bool bShouldReportPtpMismatches = (PtpUnsyncFramesPerReport >= 0) && !(PtpBaseNodeData->LastRenderedFrameNumber[0] % PtpUnsyncFramesPerReport);
+	const bool bShouldReportPtpMismatches = (PtpUnsyncFramesPerReport > 0) && !(PtpBaseNodeData->LastRenderedFrameNumber[0] % PtpUnsyncFramesPerReport);
 
 	// Iterate over the presentation requests and detect PTP de-syncs.
 	for (const TPair<FString, TArray<uint8>>& FramePresentedInfo : BarrierSyncData.RequestData)
